add more verbose comments / tooltips on the dataflow mesh boolean node and operation enum values, to help clarify what the mesh boolean options do

diff --git a/Engine/Plugins/Experimental/GeometryDataflow/Source/GeometryDataflowNodes/Public/Dataflow/MeshBooleanNodes.h b/Engine/Plugins/Experimental/GeometryDataflow/Source/GeometryDataflowNodes/Public/Dataflow/MeshBooleanNodes.h
--- a/Engine/Plugins/Experimental/GeometryDataflow/Source/GeometryDataflowNodes/Public/Dataflow/MeshBooleanNodes.h
+++ b/Engine/Plugins/Experimental/GeometryDataflow/Source/GeometryDataflowNodes/Public/Dataflow/MeshBooleanNodes.h
@@ -17,8 +17,11 @@ class UMaterialInterface;
 UENUM(BlueprintType)
 enum class EMeshBooleanOperationEnum : uint8
 {
+	// A union of A + B includes everything inside either A or B
 	Dataflow_MeshBoolean_Union UMETA(DisplayName = "Union"),
+	// An intersection of A & B includes only the points inside both A and B, i.e. trimming A by B (and vice versa)
 	Dataflow_MeshBoolean_Intersect UMETA(DisplayName = "Intersect"),
+	// A difference of A - B includes only the points inside A that are outside of B, i.e. subtracting B from A
 	Dataflow_MeshBoolean_Difference UMETA(DisplayName = "Difference"),
 	//~~~
 	//256th entry
@@ -27,7 +30,12 @@ enum class EMeshBooleanOperationEnum : uint8
 
 /**
  *
- * Mesh boolean (Union, Intersect, Difference) between two meshes
+ * Compute a Mesh Boolean between Mesh1 and Mesh2
+ * 
+ * Supported Boolean Operations:
+ *  Union (Mesh1 + Mesh2)
+ *  Difference (Mesh1 - Mesh2; removing what's inside of Mesh2 from Mesh1)
+ *  Intersection (Mesh1 & Mesh2; removing what's outside of Mesh2 from Mesh1)
  *
  */
 USTRUCT()
