Minor cleanup of dataflow mesh boolean options
- ignore returned result of Compute(), as it typically just means some edge welds failed (ie. result has small cracks)
- let user control options for welding, simplification, and which input UVs are preserved when simplified. (w/ defaults that preserve previous behavior)

#Synced-CL 49122595

diff --git a/Engine/Plugins/Experimental/GeometryDataflow/Source/GeometryDataflowNodes/Private/Dataflow/MeshBooleanNodes.cpp b/Engine/Plugins/Experimental/GeometryDataflow/Source/GeometryDataflowNodes/Private/Dataflow/MeshBooleanNodes.cpp
--- a/Engine/Plugins/Experimental/GeometryDataflow/Source/GeometryDataflowNodes/Private/Dataflow/MeshBooleanNodes.cpp
+++ b/Engine/Plugins/Experimental/GeometryDataflow/Source/GeometryDataflowNodes/Private/Dataflow/MeshBooleanNodes.cpp
@@ -61,17 +61,17 @@ void FMeshBooleanDataflowNode::Evaluate(UE::Dataflow::FContext& Context, const F
 					}
 
 					UE::Geometry::FMeshBoolean Boolean(&DynMesh1, &DynMesh2, &ResultDynMesh, BoolOp);
-					Boolean.bSimplifyAlongNewEdges = true;
-					Boolean.PreserveUVsOnlyForMesh = 0; // slight warping of the autogenerated cell UVs generally doesn't matter
-					Boolean.bWeldSharedEdges = false;
-					Boolean.bTrackAllNewEdges = true;
-
-					if (Boolean.Compute())
-					{
-						SetValue(Context, NewMesh, &Mesh);
-
-						return;
-					}
+					Boolean.bSimplifyAlongNewEdges = bSimplifyAlongNewEdges;
+					Boolean.bPreserveOverlayUVs = (PreserveUVQualityForMesh != EDataflowMeshBooleanInputMeshEnum::None);
+					Boolean.bPreserveVertexUVs = (PreserveUVQualityForMesh != EDataflowMeshBooleanInputMeshEnum::None);
+					Boolean.PreserveUVsOnlyForMesh = PreserveUVQualityForMesh == EDataflowMeshBooleanInputMeshEnum::First ? 0 :
+						((PreserveUVQualityForMesh == EDataflowMeshBooleanInputMeshEnum::Second) ? 1 : INDEX_NONE);
+					Boolean.bWeldSharedEdges = bWeldSharedEdges;
+					Boolean.bTrackAllNewEdges = false;
+
+					Boolean.Compute();
+					SetValue(Context, NewMesh, &Mesh);
+					return;
 				}
 			}
 		}
diff --git a/Engine/Plugins/Experimental/GeometryDataflow/Source/GeometryDataflowNodes/Public/Dataflow/MeshBooleanNodes.h b/Engine/Plugins/Experimental/GeometryDataflow/Source/GeometryDataflowNodes/Public/Dataflow/MeshBooleanNodes.h
--- a/Engine/Plugins/Experimental/GeometryDataflow/Source/GeometryDataflowNodes/Public/Dataflow/MeshBooleanNodes.h
+++ b/Engine/Plugins/Experimental/GeometryDataflow/Source/GeometryDataflowNodes/Public/Dataflow/MeshBooleanNodes.h
@@ -28,6 +28,16 @@ enum class EMeshBooleanOperationEnum : uint8
 	Dataflow_Max                UMETA(Hidden)
 };
 
+// Enum to select among the input meshes for a Mesh Boolean operation
+UENUM(BlueprintType)
+enum class EDataflowMeshBooleanInputMeshEnum : uint8
+{
+	None,
+	First,
+	Second,
+	Both
+};
+
 /**
  *
  * Compute a Mesh Boolean between Mesh1 and Mesh2
@@ -54,6 +64,17 @@ private:
 	UPROPERTY(EditAnywhere, Category = "Boolean");
 	EMeshBooleanOperationEnum Operation = EMeshBooleanOperationEnum::Dataflow_MeshBoolean_Intersect;
 
+	/** Whether to weld common edges at the joins of the mesh boolean */
+	UPROPERTY(EditAnywhere, Category = "Boolean", AdvancedDisplay)
+	bool bWeldSharedEdges = false;
+
+	/** Whether to simplify the result along new edges */
+	UPROPERTY(EditAnywhere, Category = "Simplification")
+	bool bSimplifyAlongNewEdges = true;
+	
+	UPROPERTY(EditAnywhere, Category = "Simplification", AdvancedDisplay, DisplayName="Preserve UV Quality for Mesh", meta=(EditCondition="bSimplifyAlongNewEdges"))
+	EDataflowMeshBooleanInputMeshEnum PreserveUVQualityForMesh = EDataflowMeshBooleanInputMeshEnum::First;
+
 	/** Mesh input */
 	UPROPERTY(meta = (DataflowInput, DataflowIntrinsic))
 	TObjectPtr<UDynamicMesh> Mesh1;
